IFC: fixed "unit scale is wrong when using inch as length measurement"(supported ifc unit conversion) and fixed "Warning: Unsupported Instance IfcSurfaceStyleElementSelect != IfcSurfaceStyleRendering"(supported StyleShading, base class for surface style)

diff --git a/Engine/Plugins/Enterprise/DatasmithIFCImporter/Source/DatasmithIFCTranslator/Private/IFC/IFCReader.cpp b/Engine/Plugins/Enterprise/DatasmithIFCImporter/Source/DatasmithIFCTranslator/Private/IFC/IFCReader.cpp
--- a/Engine/Plugins/Enterprise/DatasmithIFCImporter/Source/DatasmithIFCTranslator/Private/IFC/IFCReader.cpp
+++ b/Engine/Plugins/Enterprise/DatasmithIFCImporter/Source/DatasmithIFCTranslator/Private/IFC/IFCReader.cpp
@@ -31,6 +31,7 @@ namespace IFC
 	static const FString IfcConversionBasedUnit_Type = TEXT("IFCCONVERSIONBASEDUNIT");
 	static const FString IfcSurfaceStyle_Type = TEXT("IFCSURFACESTYLE");
 	static const FString IfcSurfaceStyleRendering_Type = TEXT("IFCSURFACESTYLERENDERING");
+	static const FString IfcSurfaceStyleShading_Type = TEXT("IFCSURFACESTYLESHADING");
 	static const FString IfcRelAssociatesMaterial_Type = TEXT("IFCRELASSOCIATESMATERIAL");
 	static const FString IfcMaterial_Type = TEXT("IFCMATERIAL");
 	static const FString IfcMaterialLayerSetUsage_Type = TEXT("IFCMATERIALLAYERSETUSAGE");
@@ -253,6 +254,7 @@ namespace IFC
 		float ParseSIPrefixText(const FString& Text)
 		{
 			if (Text.Equals(TEXT(".METRE.")))			return 1.0f;
+			else if (Text.Equals(TEXT(".INCH.")))		return 0.0254f;
 			else if (Text.Equals(TEXT(".CENTI.")))		return 0.01f;
 			else if (Text.Equals(TEXT(".MILLI.")))		return 0.001f;
 			else if (Text.Equals(TEXT(".MICRO.")))		return 0.000001f;
@@ -301,10 +303,11 @@ namespace IFC
 	void FFileReader::GetEntitiesTypes()
 	{
 		IfcMappedItemEntity = sdaiGetEntity(gIFCModel, (char*)*IfcMappedItem_Type);
-		IfcConversianBasedUnitEntity = sdaiGetEntity(gIFCModel, (char*)*IfcConversionBasedUnit_Type);
+		IfcConversionBasedUnitEntity = sdaiGetEntity(gIFCModel, (char*)*IfcConversionBasedUnit_Type);
 		IfcSIUnitEntity = sdaiGetEntity(gIFCModel, (char*)*IfcSIUnit_Type);
 		IfcSurfaceStyleEntity = sdaiGetEntity(gIFCModel, (char*)*IfcSurfaceStyle_Type);
 		IfcSurfaceStyleRenderingEntity = sdaiGetEntity(gIFCModel, (char*)*IfcSurfaceStyleRendering_Type);
+		IfcSurfaceStyleShadingEntity = sdaiGetEntity(gIFCModel, (char*)*IfcSurfaceStyleShading_Type);
 		IfcRelAssociatesMaterialEntity = sdaiGetEntity(gIFCModel, (char*)*IfcRelAssociatesMaterial_Type);
 		IfcMaterialEntity = sdaiGetEntity(gIFCModel, (char*)*IfcMaterial_Type);
 		IfcMaterialLayerSetUsageEntity = sdaiGetEntity(gIFCModel, (char*)*IfcMaterialLayerSetUsage_Type);
@@ -638,9 +641,42 @@ namespace IFC
 			int_t UnitInstance = 0;
 			engiGetAggrElement(UnitsSet, I, sdaiINSTANCE, &UnitInstance);
 
-			if (sdaiGetInstanceType(UnitInstance) == IfcConversianBasedUnitEntity)
+			if (sdaiGetInstanceType(UnitInstance) == IfcConversionBasedUnitEntity)
 			{
-				// To be considered if this could be helpful somehow.
+				// http://standards.buildingsmart.org/IFC/RELEASE/IFC2x3/TC1/HTML/ifcmeasureresource/lexical/ifcmeasurewithunit.htm
+				int_t MeasureWithUnitInstance = 0;
+				sdaiGetAttrBN(UnitInstance, (char*) *IfcConversionFactor_Name, sdaiINSTANCE, &MeasureWithUnitInstance);
+
+				if (ensure(MeasureWithUnitInstance))
+				{
+					// get SI unit from conversion unit
+					int_t SiUnitInstance = 0;
+					sdaiGetAttrBN(MeasureWithUnitInstance, (char*) *IfcUnitComponent_Name, sdaiINSTANCE, &SiUnitInstance);
+
+					int_t* ADB = 0;
+					sdaiGetAttrBN(MeasureWithUnitInstance, (char*) *IfcValueComponent_Name, sdaiADB, &ADB);
+
+					double Value = 0;
+					sdaiGetADBValue((void*)ADB, sdaiREAL, &Value);
+
+					if (sdaiGetInstanceType(SiUnitInstance) == IfcSIUnitEntity)
+					{
+						FUnit NewUnit;
+
+						Helper::GetStringAttribute(SiUnitInstance, TEXT("UnitType"), NewUnit.Type);
+						Helper::GetStringAttribute(SiUnitInstance, TEXT("Prefix"), NewUnit.Prefix);
+						Helper::GetStringAttribute(SiUnitInstance, TEXT("Name"), NewUnit.Name);
+
+						if (NewUnit.Type.Equals(TEXT(".LENGTHUNIT.")))
+						{
+							NewUnit.PrefixValue = Helper::ParseSIPrefixText(NewUnit.Prefix) * Value; // Factor in conversion value
+
+							UE_LOG(LogDatasmithIFCReader, Log, TEXT("LengthUnit = %f"), NewUnit.PrefixValue);
+						}
+
+						UnitTypesMap.Add(NewUnit.Type, NewUnit);
+					}
+				}
 			}
 			else if (sdaiGetInstanceType(UnitInstance) == IfcSIUnitEntity)
 			{
@@ -794,9 +830,23 @@ namespace IFC
 					Helper::GetADBAttribute(SurfaceStyleInstance, TEXT("specular_exponent"), OutType, specular_exponent);
 
 				}
+				else if (sdaiGetInstanceType(SurfaceStyleInstance) == IfcSurfaceStyleShadingEntity)
+				{
+					if (S == 0 && StylesEntitiesCount == 1)
+					{
+						OutMaterial.ID = SurfaceStyleInstance;
+					}
+					else
+					{
+						Messages.Emplace(EMessageSeverity::Warning, TEXT("Unsupported number of IfcSurfaceStyleElementSelect > 1 found!"));
+					}
+
+					Helper::GetRGBColorAttribute(SurfaceStyleInstance, TEXT("SurfaceColour"), OutMaterial.SurfaceColour);
+					Helper::GetFloatAttribute(SurfaceStyleInstance, TEXT("Transparency"), OutMaterial.Transparency);
+				}
 				else
 				{
-					Messages.Emplace(EMessageSeverity::Warning, TEXT("Unsupported Instance IfcSurfaceStyleElementSelect != IfcSurfaceStyleRendering"));
+					Messages.Emplace(EMessageSeverity::Warning, TEXT("Unsupported Instance of IfcSurfaceStyleElementSelect"));
 				}
 			}
 		}
diff --git a/Engine/Plugins/Enterprise/DatasmithIFCImporter/Source/DatasmithIFCTranslator/Private/IFC/IFCReader.h b/Engine/Plugins/Enterprise/DatasmithIFCImporter/Source/DatasmithIFCTranslator/Private/IFC/IFCReader.h
--- a/Engine/Plugins/Enterprise/DatasmithIFCImporter/Source/DatasmithIFCTranslator/Private/IFC/IFCReader.h
+++ b/Engine/Plugins/Enterprise/DatasmithIFCImporter/Source/DatasmithIFCTranslator/Private/IFC/IFCReader.h
@@ -211,10 +211,11 @@ namespace IFC
 		 * Loading time buffered entities to compare.
 		 */
 		int64 IfcMappedItemEntity = 0;
-		int64 IfcConversianBasedUnitEntity = 0;
+		int64 IfcConversionBasedUnitEntity = 0;
 		int64 IfcSIUnitEntity = 0;
 		int64 IfcSurfaceStyleEntity = 0;
 		int64 IfcSurfaceStyleRenderingEntity = 0;
+		int64 IfcSurfaceStyleShadingEntity = 0;
 		int64 IfcRelAssociatesMaterialEntity = 0;
 		int64 IfcMaterialEntity = 0;
 		int64 IfcMaterialLayerSetUsageEntity = 0;
