Make it clear in the API that LoadStaticMesh takes ownership of the imported meshdescription to avoid having duplicates in memory
Make the IFC importer thread-safe
Enable parallel staticmesh loading on all importers that are thread-safe
  - 7m16s [at] 85 GB Peak working set -> 5m05s [at] 67 GB when importing DATASET-0022a

From dev-enterprise[REDACTED_HANDLE]

diff --git a/Engine/Plugins/Enterprise/DatasmithC4DImporter/Source/DatasmithC4DTranslator/Private/DatasmithC4DImporter.cpp b/Engine/Plugins/Enterprise/DatasmithC4DImporter/Source/DatasmithC4DTranslator/Private/DatasmithC4DImporter.cpp
--- a/Engine/Plugins/Enterprise/DatasmithC4DImporter/Source/DatasmithC4DTranslator/Private/DatasmithC4DImporter.cpp
+++ b/Engine/Plugins/Enterprise/DatasmithC4DImporter/Source/DatasmithC4DTranslator/Private/DatasmithC4DImporter.cpp
@@ -2710,14 +2710,12 @@ TSharedPtr<IDatasmithMeshElement> FDatasmithC4DImporter::ImportMesh(melange::Pol
 	return MeshElement;
 }
 
-TArray<FMeshDescription> FDatasmithC4DImporter::GetGeometriesForMeshElement(const TSharedRef<IDatasmithMeshElement> MeshElement)
+void FDatasmithC4DImporter::GetGeometriesForMeshElementAndRelease(const TSharedRef<IDatasmithMeshElement> MeshElement, TArray<FMeshDescription>& OutMeshDescriptions)
 {
 	if (FMeshDescription* MeshDesc = MeshElementToMeshDescription.Find(&MeshElement.Get()))
 	{
-		return {*MeshDesc};
+		OutMeshDescriptions.Add(MoveTemp(*MeshDesc));
 	}
-
-	return {};
 }
 
 bool FDatasmithC4DImporter::OpenFile(const FString& InFilename)
diff --git a/Engine/Plugins/Enterprise/DatasmithC4DImporter/Source/DatasmithC4DTranslator/Private/DatasmithC4DImporter.h b/Engine/Plugins/Enterprise/DatasmithC4DImporter/Source/DatasmithC4DTranslator/Private/DatasmithC4DImporter.h
--- a/Engine/Plugins/Enterprise/DatasmithC4DImporter/Source/DatasmithC4DTranslator/Private/DatasmithC4DImporter.h
+++ b/Engine/Plugins/Enterprise/DatasmithC4DImporter/Source/DatasmithC4DTranslator/Private/DatasmithC4DImporter.h
@@ -124,7 +124,7 @@ public:
 	/** Creates a new IDatasmithMetaDataElement for the IDatasmithActorElement, adds it to the scene and returns it */
 	TSharedPtr<IDatasmithMetaDataElement> CreateMetadataForActor(const IDatasmithActorElement& Actor);
 
-	TArray<FMeshDescription> GetGeometriesForMeshElement(const TSharedRef<IDatasmithMeshElement> MeshElement);
+	void GetGeometriesForMeshElementAndRelease(const TSharedRef<IDatasmithMeshElement> MeshElement, TArray<FMeshDescription>& OutMeshDescriptions);
 	TSharedPtr<IDatasmithLevelSequenceElement> GetLevelSequence() { return LevelSequence; }
 
 	melange::BaseDocument* C4dDocument = nullptr;
diff --git a/Engine/Plugins/Enterprise/DatasmithC4DImporter/Source/DatasmithC4DTranslator/Private/DatasmithC4DTranslator.cpp b/Engine/Plugins/Enterprise/DatasmithC4DImporter/Source/DatasmithC4DTranslator/Private/DatasmithC4DTranslator.cpp
--- a/Engine/Plugins/Enterprise/DatasmithC4DImporter/Source/DatasmithC4DTranslator/Private/DatasmithC4DTranslator.cpp
+++ b/Engine/Plugins/Enterprise/DatasmithC4DImporter/Source/DatasmithC4DTranslator/Private/DatasmithC4DTranslator.cpp
@@ -18,6 +18,7 @@
 void FDatasmithC4DTranslator::Initialize(FDatasmithTranslatorCapabilities& OutCapabilities)
 {
 	OutCapabilities.bIsEnabled = true;
+	OutCapabilities.bParallelLoadStaticMeshSupported = true;
 
 	TArray<FFileFormatInfo>& Formats = OutCapabilities.SupportedFileFormats;
     Formats.Emplace(TEXT("c4d"), TEXT("Cinema 4D file format"));
@@ -42,7 +43,8 @@ bool FDatasmithC4DTranslator::LoadStaticMesh(const TSharedRef<IDatasmithMeshElem
 {
 	if (ensure(Importer.IsValid()))
 	{
-		TArray<FMeshDescription> MeshDescriptions = Importer->GetGeometriesForMeshElement(MeshElement);
+		TArray<FMeshDescription> MeshDescriptions;
+		Importer->GetGeometriesForMeshElementAndRelease(MeshElement, MeshDescriptions);
 		if (MeshDescriptions.Num() > 0)
 		{
 			OutMeshPayload.LodMeshes.Add(MoveTemp(MeshDescriptions[0]));
diff --git a/Engine/Plugins/Enterprise/DatasmithFBXImporter/Source/DatasmithDeltaGenTranslator/Private/DatasmithDeltaGenTranslator.cpp b/Engine/Plugins/Enterprise/DatasmithFBXImporter/Source/DatasmithDeltaGenTranslator/Private/DatasmithDeltaGenTranslator.cpp
--- a/Engine/Plugins/Enterprise/DatasmithFBXImporter/Source/DatasmithDeltaGenTranslator/Private/DatasmithDeltaGenTranslator.cpp
+++ b/Engine/Plugins/Enterprise/DatasmithFBXImporter/Source/DatasmithDeltaGenTranslator/Private/DatasmithDeltaGenTranslator.cpp
@@ -14,6 +14,7 @@
 void FDatasmithDeltaGenTranslator::Initialize(FDatasmithTranslatorCapabilities& OutCapabilities)
 {
 	OutCapabilities.bIsEnabled = true;
+	OutCapabilities.bParallelLoadStaticMeshSupported = true;
 
 	TArray<FFileFormatInfo>& Formats = OutCapabilities.SupportedFileFormats;
     Formats.Emplace(TEXT("fbx"), TEXT("DeltaGen Fbx files"));
@@ -80,7 +81,8 @@ bool FDatasmithDeltaGenTranslator::LoadStaticMesh(const TSharedRef<IDatasmithMes
 {
 	if (ensure(Importer.IsValid()))
 	{
-		TArray<FMeshDescription> MeshDescriptions = Importer->GetGeometriesForMeshElement(MeshElement);
+		TArray<FMeshDescription> MeshDescriptions;
+		Importer->GetGeometriesForMeshElementAndRelease(MeshElement, MeshDescriptions);
 		if (MeshDescriptions.Num() > 0)
 		{
 			OutMeshPayload.LodMeshes.Add(MoveTemp(MeshDescriptions[0]));
diff --git a/Engine/Plugins/Enterprise/DatasmithFBXImporter/Source/DatasmithFBXTranslator/Private/DatasmithFBXImporter.cpp b/Engine/Plugins/Enterprise/DatasmithFBXImporter/Source/DatasmithFBXTranslator/Private/DatasmithFBXImporter.cpp
--- a/Engine/Plugins/Enterprise/DatasmithFBXImporter/Source/DatasmithFBXTranslator/Private/DatasmithFBXImporter.cpp
+++ b/Engine/Plugins/Enterprise/DatasmithFBXImporter/Source/DatasmithFBXTranslator/Private/DatasmithFBXImporter.cpp
@@ -63,16 +63,17 @@ FDatasmithFBXImporter::~FDatasmithFBXImporter()
 {
 }
 
-TArray<FMeshDescription> FDatasmithFBXImporter::GetGeometriesForMeshElement(const TSharedRef<IDatasmithMeshElement> MeshElement)
+void FDatasmithFBXImporter::GetGeometriesForMeshElementAndRelease(const TSharedRef<IDatasmithMeshElement> MeshElement, TArray<FMeshDescription>& OutMeshDescriptions)
 {
 	TSharedPtr<FDatasmithFBXSceneMesh>* FoundMesh = MeshNameToFBXMesh.Find(MeshElement->GetName());
 	if (FoundMesh && (*FoundMesh).IsValid())
 	{
-		return {(*FoundMesh)->MeshDescription};
+		OutMeshDescriptions.Add(MoveTemp((*FoundMesh)->MeshDescription));
+	}
+	else
+	{
+		UE_LOG(LogDatasmithFBXImport, Error, TEXT("Failed to return FMeshDescription object for requested mesh element '%s'"), MeshElement->GetName());
 	}
-
-	UE_LOG(LogDatasmithFBXImport, Error, TEXT("Failed to return FMeshDescription object for requested mesh element '%s'"), MeshElement->GetName());
-	return {};
 }
 
 void FDatasmithFBXImporter::BuildAssetMaps(TSharedRef<IDatasmithScene> Scene, FActorMap& ActorsByOriginalName, FMaterialMap& MaterialsByName)
diff --git a/Engine/Plugins/Enterprise/DatasmithFBXImporter/Source/DatasmithFBXTranslator/Public/DatasmithFBXImporter.h b/Engine/Plugins/Enterprise/DatasmithFBXImporter/Source/DatasmithFBXTranslator/Public/DatasmithFBXImporter.h
--- a/Engine/Plugins/Enterprise/DatasmithFBXImporter/Source/DatasmithFBXTranslator/Public/DatasmithFBXImporter.h
+++ b/Engine/Plugins/Enterprise/DatasmithFBXImporter/Source/DatasmithFBXTranslator/Public/DatasmithFBXImporter.h
@@ -28,7 +28,7 @@ public:
 	FDatasmithFBXImporter();
 	virtual ~FDatasmithFBXImporter();
 
-	TArray<FMeshDescription> GetGeometriesForMeshElement(const TSharedRef<IDatasmithMeshElement> MeshElement);
+	void GetGeometriesForMeshElementAndRelease(const TSharedRef<IDatasmithMeshElement> MeshElement, TArray<FMeshDescription>& OutMeshDescriptions);
 
 	void BuildAssetMaps(TSharedRef<IDatasmithScene> Scene, FActorMap& ActorsByOriginalName, FMaterialMap& MaterialsByName);
 
diff --git a/Engine/Plugins/Enterprise/DatasmithFBXImporter/Source/DatasmithVREDTranslator/Private/DatasmithVREDTranslator.cpp b/Engine/Plugins/Enterprise/DatasmithFBXImporter/Source/DatasmithVREDTranslator/Private/DatasmithVREDTranslator.cpp
--- a/Engine/Plugins/Enterprise/DatasmithFBXImporter/Source/DatasmithVREDTranslator/Private/DatasmithVREDTranslator.cpp
+++ b/Engine/Plugins/Enterprise/DatasmithFBXImporter/Source/DatasmithVREDTranslator/Private/DatasmithVREDTranslator.cpp
@@ -14,6 +14,7 @@
 void FDatasmithVREDTranslator::Initialize(FDatasmithTranslatorCapabilities& OutCapabilities)
 {
 	OutCapabilities.bIsEnabled = true;
+	OutCapabilities.bParallelLoadStaticMeshSupported = true;
 
 	TArray<FFileFormatInfo>& Formats = OutCapabilities.SupportedFileFormats;
     Formats.Emplace(TEXT("fbx"), TEXT("VRED Fbx files"));
@@ -82,7 +83,8 @@ bool FDatasmithVREDTranslator::LoadStaticMesh(const TSharedRef<IDatasmithMeshEle
 {
 	if (ensure(Importer.IsValid()))
 	{
-		TArray<FMeshDescription> MeshDescriptions = Importer->GetGeometriesForMeshElement(MeshElement);
+		TArray<FMeshDescription> MeshDescriptions;
+		Importer->GetGeometriesForMeshElementAndRelease(MeshElement, MeshDescriptions);
 		if (MeshDescriptions.Num() > 0)
 		{
 			OutMeshPayload.LodMeshes.Add(MoveTemp(MeshDescriptions[0]));
diff --git a/Engine/Plugins/Enterprise/DatasmithGLTFImporter/Source/DatasmithGLTFTranslator/Private/DatasmithGLTFImporter.cpp b/Engine/Plugins/Enterprise/DatasmithGLTFImporter/Source/DatasmithGLTFTranslator/Private/DatasmithGLTFImporter.cpp
--- a/Engine/Plugins/Enterprise/DatasmithGLTFImporter/Source/DatasmithGLTFTranslator/Private/DatasmithGLTFImporter.cpp
+++ b/Engine/Plugins/Enterprise/DatasmithGLTFImporter/Source/DatasmithGLTFTranslator/Private/DatasmithGLTFImporter.cpp
@@ -318,10 +318,8 @@ bool FDatasmithGLTFImporter::SendSceneToDatasmith()
 	return true;
 }
 
-TArray<FMeshDescription> FDatasmithGLTFImporter::GetGeometriesForMeshElement(const TSharedRef<IDatasmithMeshElement> MeshElement)
+void FDatasmithGLTFImporter::GetGeometriesForMeshElementAndRelease(const TSharedRef<IDatasmithMeshElement> MeshElement, TArray<FMeshDescription>& OutMeshDescriptions)
 {
-	TArray<FMeshDescription> Result;
-
 	if (int32* MeshIndexPtr = MeshElementToGLTFMeshIndex.Find(&MeshElement.Get()))
 	{
 		int32 MeshIndex = *MeshIndexPtr;
@@ -331,10 +329,8 @@ TArray<FMeshDescription> FDatasmithGLTFImporter::GetGeometriesForMeshElement(con
 
 		StaticMeshFactory->FillMeshDescription(GLTFAsset->Meshes[MeshIndex], &MeshDescription);
 
-		Result.Add(MoveTemp(MeshDescription));
+		OutMeshDescriptions.Add(MoveTemp(MeshDescription));
 	}
-
-	return Result;
 }
 
 const TArray<TSharedRef<IDatasmithLevelSequenceElement>>& FDatasmithGLTFImporter::GetImportedSequences()
diff --git a/Engine/Plugins/Enterprise/DatasmithGLTFImporter/Source/DatasmithGLTFTranslator/Private/DatasmithGLTFImporter.h b/Engine/Plugins/Enterprise/DatasmithGLTFImporter/Source/DatasmithGLTFTranslator/Private/DatasmithGLTFImporter.h
--- a/Engine/Plugins/Enterprise/DatasmithGLTFImporter/Source/DatasmithGLTFTranslator/Private/DatasmithGLTFImporter.h
+++ b/Engine/Plugins/Enterprise/DatasmithGLTFImporter/Source/DatasmithGLTFTranslator/Private/DatasmithGLTFImporter.h
@@ -68,7 +68,7 @@ public:
 	/** Finalize import of the scene into the engine. */
 	bool SendSceneToDatasmith();
 
-	TArray<FMeshDescription> GetGeometriesForMeshElement(const TSharedRef<IDatasmithMeshElement> MeshElement);
+	void GetGeometriesForMeshElementAndRelease(const TSharedRef<IDatasmithMeshElement> MeshElement, TArray<FMeshDescription>& OutMeshDescriptions);
 
 	const TArray<TSharedRef<IDatasmithLevelSequenceElement>>& GetImportedSequences();
 
diff --git a/Engine/Plugins/Enterprise/DatasmithGLTFImporter/Source/DatasmithGLTFTranslator/Private/DatasmithGLTFTranslator.cpp b/Engine/Plugins/Enterprise/DatasmithGLTFImporter/Source/DatasmithGLTFTranslator/Private/DatasmithGLTFTranslator.cpp
--- a/Engine/Plugins/Enterprise/DatasmithGLTFImporter/Source/DatasmithGLTFTranslator/Private/DatasmithGLTFTranslator.cpp
+++ b/Engine/Plugins/Enterprise/DatasmithGLTFImporter/Source/DatasmithGLTFTranslator/Private/DatasmithGLTFTranslator.cpp
@@ -121,7 +121,8 @@ bool FDatasmithGLTFTranslator::LoadStaticMesh(const TSharedRef<IDatasmithMeshEle
 {
 	if (ensure(Importer.IsValid()))
 	{
-		TArray<FMeshDescription> MeshDescriptions = Importer->GetGeometriesForMeshElement(MeshElement);
+		TArray<FMeshDescription> MeshDescriptions;
+		Importer->GetGeometriesForMeshElementAndRelease(MeshElement, MeshDescriptions);
 		if (MeshDescriptions.Num() > 0)
 		{
 			OutMeshPayload.LodMeshes.Add(MoveTemp(MeshDescriptions[0]));
diff --git a/Engine/Plugins/Enterprise/DatasmithIFCImporter/Source/DatasmithIFCTranslator/Private/DatasmithIFCImporter.cpp b/Engine/Plugins/Enterprise/DatasmithIFCImporter/Source/DatasmithIFCTranslator/Private/DatasmithIFCImporter.cpp
--- a/Engine/Plugins/Enterprise/DatasmithIFCImporter/Source/DatasmithIFCTranslator/Private/DatasmithIFCImporter.cpp
+++ b/Engine/Plugins/Enterprise/DatasmithIFCImporter/Source/DatasmithIFCTranslator/Private/DatasmithIFCImporter.cpp
@@ -340,10 +340,8 @@ bool FDatasmithIFCImporter::SendSceneToDatasmith()
 	return true;
 }
 
-TArray<FMeshDescription> FDatasmithIFCImporter::GetGeometriesForMeshElement(const TSharedRef<IDatasmithMeshElement> MeshElement)
+void FDatasmithIFCImporter::GetGeometriesForMeshElementAndRelease(const TSharedRef<IDatasmithMeshElement> MeshElement, TArray<FMeshDescription>& OutMeshDescriptions)
 {
-	TArray<FMeshDescription> Result;
-
 	if (int64* MeshIndexPtr = MeshElementToIFCMeshIndex.Find(&MeshElement.Get()))
 	{
 #ifdef WITH_IFC_ENGINE_LIB
@@ -356,12 +354,10 @@ TArray<FMeshDescription> FDatasmithIFCImporter::GetGeometriesForMeshElement(cons
 
 			IFCStaticMeshFactory->FillMeshDescription(pIFCObject, &MeshDescription);
 
-			Result.Add(MoveTemp(MeshDescription));
+			OutMeshDescriptions.Add(MoveTemp(MeshDescription));
 		}
 #endif
 	}
-
-	return Result;
 }
 
 FString FDatasmithIFCImporter::GetMeshGlobalId(const TSharedRef<IDatasmithMeshElement> MeshElement)
@@ -397,7 +393,6 @@ void FDatasmithIFCImporter::UnloadScene()
 	ImportedMeshes.Empty();
 	ImportedActorNames.Empty();
 
-	IFCStaticMeshFactory->CleanUp();
 	IFCReader->CleanUp();
 
 	IFCReader.Reset(new IFC::FFileReader());
diff --git a/Engine/Plugins/Enterprise/DatasmithIFCImporter/Source/DatasmithIFCTranslator/Private/DatasmithIFCImporter.h b/Engine/Plugins/Enterprise/DatasmithIFCImporter/Source/DatasmithIFCTranslator/Private/DatasmithIFCImporter.h
--- a/Engine/Plugins/Enterprise/DatasmithIFCImporter/Source/DatasmithIFCTranslator/Private/DatasmithIFCImporter.h
+++ b/Engine/Plugins/Enterprise/DatasmithIFCImporter/Source/DatasmithIFCTranslator/Private/DatasmithIFCImporter.h
@@ -52,7 +52,7 @@ public:
 	void SetupCustomAnalyticsAttributes(TArray<FAnalyticsEventAttribute>& EventAttributes) const;
 
 	/**  Return all FMeshDescriptions that have been created for a particular imported IDatasmithMeshElement */
-	TArray<FMeshDescription> GetGeometriesForMeshElement(const TSharedRef<IDatasmithMeshElement> MeshElement);
+	void GetGeometriesForMeshElementAndRelease(const TSharedRef<IDatasmithMeshElement> MeshElement, TArray<FMeshDescription>& OutMeshDescriptions);
 
 	/** Mesh global ID used when reimport. */
 	FString GetMeshGlobalId(const TSharedRef<IDatasmithMeshElement> MeshElement);
diff --git a/Engine/Plugins/Enterprise/DatasmithIFCImporter/Source/DatasmithIFCTranslator/Private/DatasmithIFCTranslator.cpp b/Engine/Plugins/Enterprise/DatasmithIFCImporter/Source/DatasmithIFCTranslator/Private/DatasmithIFCTranslator.cpp
--- a/Engine/Plugins/Enterprise/DatasmithIFCImporter/Source/DatasmithIFCTranslator/Private/DatasmithIFCTranslator.cpp
+++ b/Engine/Plugins/Enterprise/DatasmithIFCImporter/Source/DatasmithIFCTranslator/Private/DatasmithIFCTranslator.cpp
@@ -40,6 +40,7 @@ void ShowLogMessages(const TArray<IFC::FLogMessage>& Errors)
 void FDatasmithIFCTranslator::Initialize(FDatasmithTranslatorCapabilities& OutCapabilities)
 {
 	OutCapabilities.bIsEnabled = true;
+	OutCapabilities.bParallelLoadStaticMeshSupported = true;
 
 	TArray<FFileFormatInfo>& Formats = OutCapabilities.SupportedFileFormats;
     Formats.Emplace(TEXT("ifc"), TEXT("IFC (Industry Foundation Classes)"));
@@ -73,7 +74,8 @@ bool FDatasmithIFCTranslator::LoadStaticMesh(const TSharedRef<IDatasmithMeshElem
 {
 	if (ensure(Importer.IsValid()))
 	{
-		TArray<FMeshDescription> MeshDescriptions = Importer->GetGeometriesForMeshElement(MeshElement);
+		TArray<FMeshDescription> MeshDescriptions;
+		Importer->GetGeometriesForMeshElementAndRelease(MeshElement, MeshDescriptions);
 		if (MeshDescriptions.Num() > 0)
 		{
 			OutMeshPayload.LodMeshes.Add(MoveTemp(MeshDescriptions[0]));
diff --git a/Engine/Plugins/Enterprise/DatasmithIFCImporter/Source/DatasmithIFCTranslator/Private/IFC/IFCStaticMeshFactory.cpp b/Engine/Plugins/Enterprise/DatasmithIFCImporter/Source/DatasmithIFCTranslator/Private/IFC/IFCStaticMeshFactory.cpp
--- a/Engine/Plugins/Enterprise/DatasmithIFCImporter/Source/DatasmithIFCTranslator/Private/IFC/IFCStaticMeshFactory.cpp
+++ b/Engine/Plugins/Enterprise/DatasmithIFCImporter/Source/DatasmithIFCTranslator/Private/IFC/IFCStaticMeshFactory.cpp
@@ -40,7 +40,7 @@ namespace IFC
 		return Hash;
 	}
 
-	void FStaticMeshFactory::FillMeshDescription(const IFC::FObject* InObject, FMeshDescription* MeshDescription)
+	void FStaticMeshFactory::FillMeshDescription(const IFC::FObject* InObject, FMeshDescription* MeshDescription) const
 	{
 		const int32 NumUVs = 1;
 
@@ -71,7 +71,8 @@ namespace IFC
 		}
 
 		// Add the PolygonGroups.
-		MaterialIndexToPolygonGroupID.Empty(10);
+		TMap<int32, FPolygonGroupID> MaterialIndexToPolygonGroupID;
+		MaterialIndexToPolygonGroupID.Reserve(10);
 		for (int32 MaterialIndex = 0; MaterialIndex < InObject->Materials.Num() || MaterialIndex < 1; MaterialIndex++)
 		{
 			const FPolygonGroupID& PolygonGroupID = MeshDescription->CreatePolygonGroup();
@@ -134,9 +135,4 @@ namespace IFC
 		ImportUniformScale = Scale;
 	}
 
-	void FStaticMeshFactory::CleanUp()
-	{
-		MaterialIndexToPolygonGroupID.Empty();
-	}
-
 }  //  namespace IFC
diff --git a/Engine/Plugins/Enterprise/DatasmithIFCImporter/Source/DatasmithIFCTranslator/Private/IFC/IFCStaticMeshFactory.h b/Engine/Plugins/Enterprise/DatasmithIFCImporter/Source/DatasmithIFCTranslator/Private/IFC/IFCStaticMeshFactory.h
--- a/Engine/Plugins/Enterprise/DatasmithIFCImporter/Source/DatasmithIFCTranslator/Private/IFC/IFCStaticMeshFactory.h
+++ b/Engine/Plugins/Enterprise/DatasmithIFCImporter/Source/DatasmithIFCTranslator/Private/IFC/IFCStaticMeshFactory.h
@@ -26,7 +26,7 @@ namespace IFC
 
 		inline const TArray<FLogMessage>& GetLogMessages() const { return Messages; }
 
-		void FillMeshDescription(const IFC::FObject* InObject, FMeshDescription* MeshDescription);
+		void FillMeshDescription(const IFC::FObject* InObject, FMeshDescription* OutMeshDescription) const;
 
 		FMD5Hash ComputeHash(const IFC::FObject& InObject);
 
@@ -35,8 +35,6 @@ namespace IFC
 
 		void SetReserveSize(uint32 Size);
 
-		void CleanUp();
-
 	protected:
 		mutable TArray<FLogMessage> Messages;
 
@@ -44,8 +42,6 @@ namespace IFC
 
 	private:
 		using FIndexVertexIdMap = TMap<int32, FVertexID>;
-
-		TMap<int32, FPolygonGroupID> MaterialIndexToPolygonGroupID;
 	};
 
 }  // namespace IFC
