Fix runtime issue with IFC translator

diff --git a/Engine/Plugins/Enterprise/DatasmithIFCImporter/Source/DatasmithIFCTranslator/DatasmithIFCTranslator.Build.cs b/Engine/Plugins/Enterprise/DatasmithIFCImporter/Source/DatasmithIFCTranslator/DatasmithIFCTranslator.Build.cs
--- a/Engine/Plugins/Enterprise/DatasmithIFCImporter/Source/DatasmithIFCTranslator/DatasmithIFCTranslator.Build.cs
+++ b/Engine/Plugins/Enterprise/DatasmithIFCImporter/Source/DatasmithIFCTranslator/DatasmithIFCTranslator.Build.cs
@@ -52,6 +52,11 @@ namespace UnrealBuildTool.Rules
                     "DatasmithImporter"
                 }
             );
-        }
-    }
+
+			if (Target.Platform == UnrealTargetPlatform.Win64)
+			{
+				RuntimeDependencies.Add("$(EngineDir)/Plugins/Enterprise/DatasmithIFCImporter/Binaries/Win64/ifcengine.dll");
+			}
+		}
+	}
 }
