IFC: imported normals

fixes "random triangles show up black in IFC import"

diff --git a/Engine/Plugins/Enterprise/DatasmithIFCImporter/Source/DatasmithIFCTranslator/Private/IFC/IFCStaticMeshFactory.cpp b/Engine/Plugins/Enterprise/DatasmithIFCImporter/Source/DatasmithIFCTranslator/Private/IFC/IFCStaticMeshFactory.cpp
--- a/Engine/Plugins/Enterprise/DatasmithIFCImporter/Source/DatasmithIFCTranslator/Private/IFC/IFCStaticMeshFactory.cpp
+++ b/Engine/Plugins/Enterprise/DatasmithIFCImporter/Source/DatasmithIFCTranslator/Private/IFC/IFCStaticMeshFactory.cpp
@@ -104,6 +104,9 @@ namespace IFC
 
 				VertexInstanceIDs.Add(VertexInstanceID);
 				VertexIDs.Add(VertexID);
+
+				const float* Vertex = &(InObject->facesVertices[(VertexIndex * (InObject->vertexElementSize / sizeof(float)))]);
+				VertexInstanceNormals.Set(VertexInstanceID, FVector(Vertex[3], Vertex[4], Vertex[5]));
 			}
 
 			bool bIsWrong = false;
