IFC: Fix transform applied to imported normals #jira UE-81472 #rb JeanMichel.Dignard, Kerim.Borchaev

diff --git a/Engine/Plugins/Enterprise/DatasmithIFCImporter/Source/DatasmithIFCTranslator/Private/IFC/IFCReader.cpp b/Engine/Plugins/Enterprise/DatasmithIFCImporter/Source/DatasmithIFCTranslator/Private/IFC/IFCReader.cpp
--- a/Engine/Plugins/Enterprise/DatasmithIFCImporter/Source/DatasmithIFCTranslator/Private/IFC/IFCReader.cpp
+++ b/Engine/Plugins/Enterprise/DatasmithIFCImporter/Source/DatasmithIFCTranslator/Private/IFC/IFCReader.cpp
@@ -1474,14 +1474,16 @@ namespace IFC
 			}
 
 			{
-				// Matrix to be used as UE transform:
-				// Negated axes Y coordinates - to convert IFC transform to UE(right-handed to left-handed).
-				// And Axis1(for Y) is negated to (later) flip IFC meshes around Y axis so they don't look 'mirrored' in the Content Viewer
+				// Convert the object transform from IFC (right-handed) to UE4 (left-handed)
+				// Another way of thinking about this: The final transform should converts to IFC space, apply the object
+				// transform, then converts back to UE4 space. The first conversion is a pre-multiply and flips a Y column
+				// (or row), and the second conversion is a post-multiply and flips the Y row (or column),
+				// hence the pattern of minus signs
 				FMatrix Matrix = FMatrix::Identity;
-				FVector Axis0 = FVector(transformationMatrix[0], -transformationMatrix[1], transformationMatrix[2]);
-				FVector Axis1 = -FVector(transformationMatrix[3], -transformationMatrix[4], transformationMatrix[5]);
-				FVector Axis2 = FVector(transformationMatrix[6], -transformationMatrix[7], transformationMatrix[8]);
-				FVector Origin = FVector(transformationMatrix[9], -transformationMatrix[10], transformationMatrix[11]);
+				FVector Axis0  = FVector(  transformationMatrix[0], - transformationMatrix[1],   transformationMatrix[2]);
+				FVector Axis1  = FVector(- transformationMatrix[3],   transformationMatrix[4], - transformationMatrix[5]);
+				FVector Axis2  = FVector(  transformationMatrix[6], - transformationMatrix[7],   transformationMatrix[8]);
+				FVector Origin = FVector(  transformationMatrix[9], - transformationMatrix[10],  transformationMatrix[11]);
 				Matrix.SetAxes(&Axis0, &Axis1, &Axis2, &Origin);
 				InObject->Transform.SetFromMatrix(Matrix);
 			}
@@ -1599,7 +1601,7 @@ namespace IFC
 	void FFileReader::GatherObjects(const FString& InName, bool bVisible, int64 iCircleSegments)
 	{
 		int_t* ObjectEntities = sdaiGetEntityExtentBN(gIFCModel, (char *)*InName);
-		
+
 		GatherObjects(gIFCModel, sdaiGetEntity(gIFCModel, (char *)*InName), ObjectEntities, bVisible, iCircleSegments);
 	}
 
diff --git a/Engine/Plugins/Enterprise/DatasmithIFCImporter/Source/DatasmithIFCTranslator/Private/IFC/IFCStaticMeshFactory.cpp b/Engine/Plugins/Enterprise/DatasmithIFCImporter/Source/DatasmithIFCTranslator/Private/IFC/IFCStaticMeshFactory.cpp
--- a/Engine/Plugins/Enterprise/DatasmithIFCImporter/Source/DatasmithIFCTranslator/Private/IFC/IFCStaticMeshFactory.cpp
+++ b/Engine/Plugins/Enterprise/DatasmithIFCImporter/Source/DatasmithIFCTranslator/Private/IFC/IFCStaticMeshFactory.cpp
@@ -63,7 +63,7 @@ namespace IFC
 		{
 			const FVertexID& VertexID = MeshDescription->CreateVertex();
 			VertexPositions[VertexID].X = InObject->facesVertices[(I * (InObject->vertexElementSize / sizeof(float))) + 0];
-			// mirror around Y axis 
+			// Flip Y to keep mesh looking the same as the coordinate system changes from RH -> LH
 			VertexPositions[VertexID].Y = - InObject->facesVertices[(I * (InObject->vertexElementSize / sizeof(float))) + 1];
 			VertexPositions[VertexID].Z = InObject->facesVertices[(I * (InObject->vertexElementSize / sizeof(float))) + 2];
 			VertexPositions[VertexID] *= ImportUniformScale;
@@ -83,6 +83,9 @@ namespace IFC
 
 		bool bMeshUsesEmptyMaterial = false;
 		bool bDidGenerateTexCoords = false;
+
+		FTransform WorldToObject = InObject->Transform.Inverse();
+
 		for (int32 Index = 0; Index < InObject->TrianglesArray.Num(); ++Index)
 		{
 			const IFC::FPolygon* IFCPolygon = &InObject->TrianglesArray[Index];
@@ -106,7 +109,10 @@ namespace IFC
 				VertexIDs.Add(VertexID);
 
 				const float* Vertex = &(InObject->facesVertices[(VertexIndex * (InObject->vertexElementSize / sizeof(float)))]);
-				VertexInstanceNormals.Set(VertexInstanceID, FVector(Vertex[3], Vertex[4], Vertex[5]));
+
+				// Flip Y to go from RH -> LH
+				FVector Normal = WorldToObject.TransformVector(FVector(Vertex[3], -Vertex[4], Vertex[5]));
+				VertexInstanceNormals.Set(VertexInstanceID, Normal.GetSafeNormal());
 			}
 
 			bool bIsWrong = false;
