IFC: fixed crash when importing IFC file

- IFC4 file used entity type that was removed in that revision

diff --git a/Engine/Plugins/Enterprise/DatasmithIFCImporter/Source/DatasmithIFCTranslator/Private/IFC/IFCReader.cpp b/Engine/Plugins/Enterprise/DatasmithIFCImporter/Source/DatasmithIFCTranslator/Private/IFC/IFCReader.cpp
--- a/Engine/Plugins/Enterprise/DatasmithIFCImporter/Source/DatasmithIFCTranslator/Private/IFC/IFCReader.cpp
+++ b/Engine/Plugins/Enterprise/DatasmithIFCImporter/Source/DatasmithIFCTranslator/Private/IFC/IFCReader.cpp
@@ -511,16 +511,22 @@ namespace IFC
 							engiGetAggrElement(ChildInstances, ChildInstanceIndex, sdaiINSTANCE, &ChildInstance);
 
 							int32* ObjectIndexPtr = IFCInstanceToObjectIdMap.Find(ChildInstance);
-							check(ObjectIndexPtr && IFCObjects.IsValidIndex(*ObjectIndexPtr));
-
-							Object.DecomposedBy.Add(*ObjectIndexPtr);
+							if (ObjectIndexPtr && IFCObjects.IsValidIndex(*ObjectIndexPtr))
+							{
+								Object.DecomposedBy.Add(*ObjectIndexPtr);
 
-							GatherSpatialHierarchy(InIFCModel, ChildInstance);
+								GatherSpatialHierarchy(InIFCModel, ChildInstance);
+							}
+							else
+							{
+								TCHAR* EntityNamePtr = nullptr;
+								engiGetEntityName(sdaiGetInstanceType(ChildInstance), sdaiUNICODE, (char**)&EntityNamePtr);
+								FString EntityName = EntityNamePtr;
+								Messages.Emplace(EMessageSeverity::Warning, TEXT("Undefined Entity Type: ") + EntityName);
+							}
 						}
 					}
-
 				}
-
 			}
 		}
 
@@ -1591,6 +1597,7 @@ namespace IFC
 		for (int_t EntityIndex = 0; EntityIndex < EntityCount; ++EntityIndex)
 		{
 			int_t Entity = engiGetEntityElement(InIFCModel, EntityIndex);
+
 			if (engiGetEntityParent(Entity) == InParentEntity)
 			{
 				GatherObjectsEntityHierarchy(InIFCModel, Entity, bVisible, CircleSegments);
