Cleanup literal references to UE4

diff --git a/Engine/Plugins/Editor/GLTFImporter/Source/GLTFCore/Private/GLTF/ConversionUtilities.h b/Engine/Plugins/Editor/GLTFImporter/Source/GLTFCore/Private/GLTF/ConversionUtilities.h
--- a/Engine/Plugins/Editor/GLTFImporter/Source/GLTFCore/Private/GLTF/ConversionUtilities.h
+++ b/Engine/Plugins/Editor/GLTFImporter/Source/GLTFCore/Private/GLTF/ConversionUtilities.h
@@ -17,7 +17,7 @@ namespace GLTF
 	inline FVector ConvertVec3(const FVector& Vec)
 	{
 		// glTF uses a right-handed coordinate system, with Y up.
-		// UE4 uses a left-handed coordinate system, with Z up.
+		// Unreal uses a left-handed coordinate system, with Z up.
 		return {Vec.X, Vec.Z, Vec.Y};
 	}
 
@@ -33,7 +33,7 @@ namespace GLTF
 	inline FQuat ConvertQuat(const FQuat& Quat)
 	{
 		// glTF uses a right-handed coordinate system, with Y up.
-		// UE4 uses a left-handed coordinate system, with Z up.
+		// Unreal uses a left-handed coordinate system, with Z up.
 		// Quat = (qX, qY, qZ, qW) = (sin(angle/2) * aX, sin(angle/2) * aY, sin(angle/2) * aZ, cons(angle/2))
 		// where (aX, aY, aZ) - rotation axis, angle - rotation angle
 		// Y swapped with Z between these coordinate systems
diff --git a/Engine/Plugins/Enterprise/DataprepEditor/Source/DataprepEditor/Private/DataprepEditor.cpp b/Engine/Plugins/Enterprise/DataprepEditor/Source/DataprepEditor/Private/DataprepEditor.cpp
--- a/Engine/Plugins/Enterprise/DataprepEditor/Source/DataprepEditor/Private/DataprepEditor.cpp
+++ b/Engine/Plugins/Enterprise/DataprepEditor/Source/DataprepEditor/Private/DataprepEditor.cpp
@@ -464,7 +464,7 @@ void FDataprepEditor::CleanUpTemporaryDirectories()
 					// Delete directories if process is not valid
 					bool bDeleteDirectories = !ProcHandle.IsValid();
 
-					// Process is valid, check if application associated with process id is UE4 editor
+					// Process is valid, check if application associated with process id is UnrealEditor
 					if(!bDeleteDirectories)
 					{
 						const FString ApplicationName = FPlatformProcess::GetApplicationName( ProcessID );
diff --git a/Engine/Plugins/Enterprise/DatasmithCADImporter/Source/DatasmithCADTranslator/Private/DatasmithSceneGraphBuilder.cpp b/Engine/Plugins/Enterprise/DatasmithCADImporter/Source/DatasmithCADTranslator/Private/DatasmithSceneGraphBuilder.cpp
--- a/Engine/Plugins/Enterprise/DatasmithCADImporter/Source/DatasmithCADTranslator/Private/DatasmithSceneGraphBuilder.cpp
+++ b/Engine/Plugins/Enterprise/DatasmithCADImporter/Source/DatasmithCADTranslator/Private/DatasmithSceneGraphBuilder.cpp
@@ -73,13 +73,13 @@ namespace
 
 
 FDatasmithSceneGraphBuilder::FDatasmithSceneGraphBuilder(
-	TMap<uint32, FString>& InCADFileToUE4FileMap, 
+	TMap<uint32, FString>& InCADFileToUnrealFileMap, 
 	const FString& InCachePath, 
 	TSharedRef<IDatasmithScene> InScene, 
 	const FDatasmithSceneSource& InSource, 
 	const CADLibrary::FImportParameters& InImportParameters)
 		: FDatasmithSceneBaseGraphBuilder(nullptr, InScene, InSource, InImportParameters)
-		, CADFileToSceneGraphDescriptionFile(InCADFileToUE4FileMap)
+		, CADFileToSceneGraphDescriptionFile(InCADFileToUnrealFileMap)
 		, CachePath(InCachePath)
 {
 }
diff --git a/Engine/Plugins/Enterprise/DatasmithIFCImporter/Source/DatasmithIFCTranslator/Private/IFC/IFCReader.cpp b/Engine/Plugins/Enterprise/DatasmithIFCImporter/Source/DatasmithIFCTranslator/Private/IFC/IFCReader.cpp
--- a/Engine/Plugins/Enterprise/DatasmithIFCImporter/Source/DatasmithIFCTranslator/Private/IFC/IFCReader.cpp
+++ b/Engine/Plugins/Enterprise/DatasmithIFCImporter/Source/DatasmithIFCTranslator/Private/IFC/IFCReader.cpp
@@ -1481,9 +1481,9 @@ namespace IFC
 			}
 
 			{
-				// Convert the object transform from IFC (right-handed) to UE4 (left-handed)
+				// Convert the object transform from IFC (right-handed) to Unreal (left-handed)
 				// Another way of thinking about this: The final transform should converts to IFC space, apply the object
-				// transform, then converts back to UE4 space. The first conversion is a pre-multiply and flips a Y column
+				// transform, then converts back to Unreal space. The first conversion is a pre-multiply and flips a Y column
 				// (or row), and the second conversion is a post-multiply and flips the Y row (or column),
 				// hence the pattern of minus signs
 				FMatrix Matrix = FMatrix::Identity;
