Add deprecated in IFC plugin

diff --git a/Engine/Plugins/Enterprise/DatasmithIFCImporter/DatasmithIFCImporter.uplugin b/Engine/Plugins/Enterprise/DatasmithIFCImporter/DatasmithIFCImporter.uplugin
--- a/Engine/Plugins/Enterprise/DatasmithIFCImporter/DatasmithIFCImporter.uplugin
+++ b/Engine/Plugins/Enterprise/DatasmithIFCImporter/DatasmithIFCImporter.uplugin
@@ -2,8 +2,8 @@
 	"FileVersion": 3,
 	"Version": 1.0,
 	"VersionName": "1.0",
-	"FriendlyName": "Datasmith IFC Importer",
-	"Description": "An importer for IFC 3D geometry.",
+	"FriendlyName": "Datasmith IFC Importer (deprecated)",
+	"Description": "An importer for IFC 3D geometry (deprecated, use Datasmith CAD Importer instead). ",
 	"Category": "Importers",
 	"CreatedBy": "Epic Games, Inc.",
 	"CreatedByURL": "http://epicgames.com",
