Fixed crash when loading Datasmith supported file formats  with options.

diff --git a/Engine/Plugins/Enterprise/DatasmithC4DImporter/Source/DatasmithC4DTranslator/Private/DatasmithC4DTranslator.cpp b/Engine/Plugins/Enterprise/DatasmithC4DImporter/Source/DatasmithC4DTranslator/Private/DatasmithC4DTranslator.cpp
--- a/Engine/Plugins/Enterprise/DatasmithC4DImporter/Source/DatasmithC4DTranslator/Private/DatasmithC4DTranslator.cpp
+++ b/Engine/Plugins/Enterprise/DatasmithC4DImporter/Source/DatasmithC4DTranslator/Private/DatasmithC4DTranslator.cpp
@@ -91,7 +91,7 @@ bool FDatasmithC4DTranslator::LoadLevelSequence(const TSharedRef<IDatasmithLevel
 
 void FDatasmithC4DTranslator::GetSceneImportOptions(TArray<TObjectPtr<UDatasmithOptionsBase>>& Options)
 {
-	Options.Add(GetOrCreateC4DImportOptions().Get());
+	Options.Add(GetOrCreateC4DImportOptions());
 }
 
 void FDatasmithC4DTranslator::SetSceneImportOptions(const TArray<TObjectPtr<UDatasmithOptionsBase>>& Options)
@@ -100,7 +100,7 @@ void FDatasmithC4DTranslator::SetSceneImportOptions(const TArray<TObjectPtr<UDat
 	{
 		if (UDatasmithC4DImportOptions* InImportOptions = Cast<UDatasmithC4DImportOptions>(OptionPtr))
 		{
-			ImportOptions.Reset(InImportOptions);
+			ImportOptions = InImportOptions;
 		}
 	}
 
@@ -119,11 +119,11 @@ void FDatasmithC4DTranslator::SetSceneImportOptions(const TArray<TObjectPtr<UDat
 	}
 }
 
-TStrongObjectPtr<UDatasmithC4DImportOptions>& FDatasmithC4DTranslator::GetOrCreateC4DImportOptions()
+TObjectPtr<UDatasmithC4DImportOptions>& FDatasmithC4DTranslator::GetOrCreateC4DImportOptions()
 {
-	if (!ImportOptions.IsValid())
+	if (!ImportOptions)
 	{
-		ImportOptions = Datasmith::MakeOptions<UDatasmithC4DImportOptions>();
+		ImportOptions = Datasmith::MakeOptionsObjectPtr<UDatasmithC4DImportOptions>();
 	}
 	return ImportOptions;
 }
diff --git a/Engine/Plugins/Enterprise/DatasmithC4DImporter/Source/DatasmithC4DTranslator/Private/DatasmithC4DTranslator.h b/Engine/Plugins/Enterprise/DatasmithC4DImporter/Source/DatasmithC4DTranslator/Private/DatasmithC4DTranslator.h
--- a/Engine/Plugins/Enterprise/DatasmithC4DImporter/Source/DatasmithC4DTranslator/Private/DatasmithC4DTranslator.h
+++ b/Engine/Plugins/Enterprise/DatasmithC4DImporter/Source/DatasmithC4DTranslator/Private/DatasmithC4DTranslator.h
@@ -33,11 +33,11 @@ public:
 	static FPreTranslateEvent& OnPreTranslate() { return PreTranslateEvent; }*/
 	
 private:
-	TStrongObjectPtr<UDatasmithC4DImportOptions>& GetOrCreateC4DImportOptions();
+	TObjectPtr<UDatasmithC4DImportOptions>& GetOrCreateC4DImportOptions();
 
 private:
 	/** -1: try load dynamic initially; 0: dynamic cannot be loaded; 1: dynamic successful loaded */
-	TStrongObjectPtr<UDatasmithC4DImportOptions> ImportOptions;
+	TObjectPtr<UDatasmithC4DImportOptions> ImportOptions;
 	TSharedPtr<FDatasmithC4DImporter> Importer;
 #endif
 };
diff --git a/Engine/Plugins/Enterprise/DatasmithIFCImporter/Source/DatasmithIFCTranslator/Private/DatasmithIFCTranslator.cpp b/Engine/Plugins/Enterprise/DatasmithIFCImporter/Source/DatasmithIFCTranslator/Private/DatasmithIFCTranslator.cpp
--- a/Engine/Plugins/Enterprise/DatasmithIFCImporter/Source/DatasmithIFCTranslator/Private/DatasmithIFCTranslator.cpp
+++ b/Engine/Plugins/Enterprise/DatasmithIFCImporter/Source/DatasmithIFCTranslator/Private/DatasmithIFCTranslator.cpp
@@ -87,7 +87,7 @@ bool FDatasmithIFCTranslator::LoadScene(TSharedRef<IDatasmithScene> OutScene)
 	OutScene->SetHost(TEXT("IFCTranslator"));
 	OutScene->SetProductName(TEXT("IFC"));
 
-    Importer = MakeShared<FDatasmithIFCImporter>(OutScene, ImportOptions.Get());
+    Importer = MakeShared<FDatasmithIFCImporter>(OutScene, ImportOptions);
 
 	const FString& FilePath = GetSource().GetSourceFile();
 	if(!Importer->OpenFile(FilePath))
@@ -132,12 +132,12 @@ bool FDatasmithIFCTranslator::LoadLevelSequence(const TSharedRef<IDatasmithLevel
 
 void FDatasmithIFCTranslator::GetSceneImportOptions(TArray<TObjectPtr<UDatasmithOptionsBase>>& Options)
 {
-	if (!ImportOptions.IsValid())
+	if (!ImportOptions)
 	{
-		ImportOptions = Datasmith::MakeOptions<UDatasmithIFCImportOptions>();
+		ImportOptions = Datasmith::MakeOptionsObjectPtr<UDatasmithIFCImportOptions>();
 	}
 
-	Options.Add(ImportOptions.Get());
+	Options.Add(ImportOptions);
 }
 
 void FDatasmithIFCTranslator::SetSceneImportOptions(const TArray<TObjectPtr<UDatasmithOptionsBase>>& Options)
@@ -146,12 +146,12 @@ void FDatasmithIFCTranslator::SetSceneImportOptions(const TArray<TObjectPtr<UDat
 	{
 		if (UDatasmithIFCImportOptions* InImportOptions = Cast<UDatasmithIFCImportOptions>(OptionPtr))
 		{
-			ImportOptions.Reset(InImportOptions);
+			ImportOptions = InImportOptions;
 		}
 	}
 
-	if (Importer.IsValid())
+	if (Importer)
 	{
-		Importer->SetImportOptions(ImportOptions.Get());
+		Importer->SetImportOptions(ImportOptions);
 	}
 }
diff --git a/Engine/Plugins/Enterprise/DatasmithIFCImporter/Source/DatasmithIFCTranslator/Private/DatasmithIFCTranslator.h b/Engine/Plugins/Enterprise/DatasmithIFCImporter/Source/DatasmithIFCTranslator/Private/DatasmithIFCTranslator.h
--- a/Engine/Plugins/Enterprise/DatasmithIFCImporter/Source/DatasmithIFCTranslator/Private/DatasmithIFCTranslator.h
+++ b/Engine/Plugins/Enterprise/DatasmithIFCImporter/Source/DatasmithIFCTranslator/Private/DatasmithIFCTranslator.h
@@ -24,6 +24,6 @@ public:
 	virtual void SetSceneImportOptions(const TArray<TObjectPtr<UDatasmithOptionsBase>>& Options) override;
 
 private:
-    TStrongObjectPtr<UDatasmithIFCImportOptions> ImportOptions;
+	TObjectPtr<UDatasmithIFCImportOptions> ImportOptions;
     TSharedPtr<class FDatasmithIFCImporter> Importer;
 };
diff --git a/Engine/Plugins/Experimental/Enterprise/DatasmithInterchange/Source/DatasmithInterchange/Private/InterchangeDatasmithPipeline.cpp b/Engine/Plugins/Experimental/Enterprise/DatasmithInterchange/Source/DatasmithInterchange/Private/InterchangeDatasmithPipeline.cpp
--- a/Engine/Plugins/Experimental/Enterprise/DatasmithInterchange/Source/DatasmithInterchange/Private/InterchangeDatasmithPipeline.cpp
+++ b/Engine/Plugins/Experimental/Enterprise/DatasmithInterchange/Source/DatasmithInterchange/Private/InterchangeDatasmithPipeline.cpp
@@ -53,6 +53,9 @@ UInterchangeDatasmithPipeline::UInterchangeDatasmithPipeline()
 
 	MeshPipeline->CommonMeshesProperties = CommonMeshesProperties;
 	MeshPipeline->CommonSkeletalMeshesAndAnimationsProperties = CommonSkeletalMeshesAndAnimationsProperties;
+
+	AnimationPipeline->CommonMeshesProperties = CommonMeshesProperties;
+	AnimationPipeline->CommonSkeletalMeshesAndAnimationsProperties = CommonSkeletalMeshesAndAnimationsProperties;
 }
 
 void UInterchangeDatasmithPipeline::ExecutePipeline(UInterchangeBaseNodeContainer* InBaseNodeContainer, const TArray<UInterchangeSourceData*>& SourceDatas)
